[Programs] Add TraceQuery, a CLI tool for querying .utrace files as JSONL

TraceQuery is a lightweight standalone program that reads Unreal Engine .utrace
binary trace files and outputs structured JSONL to stdout. It enables scripted
performance analysis workflows without requiring UnrealInsights or the editor.

== C++ (Engine/Source/Programs/TraceQuery) ==

Standalone UBT program target (monolithic, console app, no engine linkage).
Depends only on Core, TraceAnalysis, and TraceServices.

FJsonTraceAnalyzer implements UE::Trace::IAnalyzer and decodes:
  - Stats::Spec + Stats::EventBatch2 - named stat values (cycle counters, etc.)
  - Counters::Spec + SetValueInt/Float - named counter time series
  - Diagnostics::Session2 - session metadata (platform, CL, branch, config)
  - Misc::BeginFrame - frame boundary markers
  - CpuProfiler::EventBatch (V2/V3) - CPU scope enter/leave with timing

CPU scope events are buffered per (thread, specId) and flushed as grouped arrays
in OnAnalysisEnd, producing one JSONL object per group with all [time, duration]
data points. This avoids per-event output for high-frequency scopes.

CLI supports filtering by logger, event name, stat/counter substring, and
time range (--start-time / --end-time in seconds).

== Intended usage ==

  TraceQuery.exe path/to/file.utrace --stat "Mass Total" --start-time 60

#Synced-CL 51590771

diff --git a/Engine/Source/Programs/TraceQuery/Private/JsonTraceAnalyzer.cpp b/Engine/Source/Programs/TraceQuery/Private/JsonTraceAnalyzer.cpp
new file mode 100644
--- /dev/null
+++ b/Engine/Source/Programs/TraceQuery/Private/JsonTraceAnalyzer.cpp
@@ -0,0 +1,783 @@
+// Copyright Epic Games, Inc. All Rights Reserved.
+
+#include "JsonTraceAnalyzer.h"
+
+#include "Containers/ArrayView.h"
+
+// Public TraceServices utilities (LegacyAttachmentArray, LegacyAttachmentString)
+#include "TraceServices/Utils.h"
+
+#include <stdio.h>
+
+////////////////////////////////////////////////////////////////////////////////////////////////////
+// Decode utilities — duplicated from TraceServices/Private/Common/Utils.h
+// These are private to TraceServices, so we inline them here (same approach as
+// ConvertToText.cpp in TraceAnalyzer).
+// Only needed for Stats batch decode (CpuProfiler and Counters now use TraceServices analyzers).
+////////////////////////////////////////////////////////////////////////////////////////////////////
+
+namespace TraceQueryLocal
+{
+
+inline uint64 Decode7bit(const uint8*& BufferPtr)
+{
+	uint64 Value = 0;
+	uint64 ByteIndex = 0;
+	bool bHasMoreBytes;
+	do
+	{
+		uint8 ByteValue = *BufferPtr++;
+		bHasMoreBytes = ByteValue & 0x80;
+		Value |= uint64(ByteValue & 0x7f) << (ByteIndex * 7);
+		++ByteIndex;
+	} while (bHasMoreBytes);
+	return Value;
+}
+
+inline int64 DecodeZigZag(const uint8*& BufferPtr)
+{
+	uint64 Z = Decode7bit(BufferPtr);
+	return (Z & 1) ? (Z >> 1) ^ -1 : (Z >> 1);
+}
+
+inline uint32 GetThreadIdField(
+	const UE::Trace::IAnalyzer::FOnEventContext& Context,
+	const ANSICHAR* FieldName = "ThreadId")
+{
+	static const uint32 Bias = 0x70000000;
+	uint32 ThreadId = Context.EventData.GetValue<uint32>(FieldName, 0);
+	ThreadId |= ThreadId ? Bias : Context.ThreadInfo.GetId();
+	return ThreadId;
+}
+
+} // namespace TraceQueryLocal
+
+////////////////////////////////////////////////////////////////////////////////////////////////////
+// FTraceQueryCpuProfilerProvider
+////////////////////////////////////////////////////////////////////////////////////////////////////
+
+const FString* FTraceQueryCpuProfilerProvider::LookupTimerName(uint32 TimerId) const
+{
+	// Handle metadata bit-inversion: negative int32 means metadata entry
+	if (int32(TimerId) < 0)
+	{
+		uint32 MetaIdx = ~TimerId;
+		if (MetaIdx < uint32(Metadatas.Num()))
+		{
+			TimerId = Metadatas[MetaIdx].TimerId;
+		}
+		else
+		{
+			return nullptr;
+		}
+	}
+
+	if (TimerId < uint32(TimerNames.Num()) && TimerNames[TimerId].IsSet())
+	{
+		return &TimerNames[TimerId].GetValue();
+	}
+	return nullptr;
+}
+
+void FTraceQueryCpuProfilerProvider::AddThread(uint32 Id, const TCHAR* Name, EThreadPriority Priority)
+{
+	// Ensure thread exists in our map
+	if (!Threads.Contains(Id))
+	{
+		Threads.Add(Id, MakeUnique<FThread>(Id, this));
+	}
+}
+
+uint32 FTraceQueryCpuProfilerProvider::AddTimer(TraceServices::ETimingProfilerTimerType Type)
+{
+	if (Type == TraceServices::ETimingProfilerTimerType::CPU)
+	{
+		return AddTimerInternal(FStringView());
+	}
+	return 0;
+}
+
+uint32 FTraceQueryCpuProfilerProvider::AddCpuTimer(FStringView Name, const TCHAR* File, uint32 Line)
+{
+	return AddTimerInternal(Name);
+}
+
+uint32 FTraceQueryCpuProfilerProvider::AddTimerInternal(FStringView Name)
+{
+	TOptional<FString> TimerName;
+	if (!Name.IsEmpty())
+	{
+		TimerName.Emplace(FString(Name));
+	}
+
+	uint32 TimerId = TimerNames.Add(TimerName);
+
+	if (NumCpuSpecs)
+	{
+		++(*NumCpuSpecs);
+	}
+
+	return TimerId;
+}
+
+void FTraceQueryCpuProfilerProvider::SetTimerName(uint32 TimerId, FStringView Name)
+{
+	check(TimerId < uint32(TimerNames.Num()));
+	TimerNames[TimerId].Emplace(FString(Name));
+}
+
+uint32 FTraceQueryCpuProfilerProvider::AddMetadata(uint32 MainTimerId, TArray<uint8>&& Metadata)
+{
+	uint32 MetadataId = Metadatas.Num();
+	Metadatas.Add({ MoveTemp(Metadata), MainTimerId });
+	return ~MetadataId;
+}
+
+TArrayView<uint8> FTraceQueryCpuProfilerProvider::GetEditableMetadata(uint32 TimerId)
+{
+	if (int32(TimerId) >= 0)
+	{
+		return TArrayView<uint8>();
+	}
+
+	uint32 MetaIdx = ~TimerId;
+	if (MetaIdx >= uint32(Metadatas.Num()))
+	{
+		return TArrayView<uint8>();
+	}
+
+	return Metadatas[MetaIdx].Payload;
+}
+
+TraceServices::IEditableTimeline<TraceServices::FTimingProfilerEvent>& FTraceQueryCpuProfilerProvider::GetCpuThreadEditableTimeline(uint32 ThreadId)
+{
+	TUniquePtr<FThread>* Found = Threads.Find(ThreadId);
+	if (Found)
+	{
+		return *(Found->Get());
+	}
+	return *Threads.Add(ThreadId, MakeUnique<FThread>(ThreadId, this));
+}
+
+void FTraceQueryCpuProfilerProvider::FThread::AppendBeginEvent(double StartTime, const TraceServices::FTimingProfilerEvent& Event)
+{
+	FScopeEnter Entry;
+	Entry.TimerId = Event.TimerIndex;
+	Entry.StartTime = StartTime;
+	ScopeStack.Push(Entry);
+}
+
+void FTraceQueryCpuProfilerProvider::FThread::AppendEndEvent(double EndTime)
+{
+	if (ScopeStack.IsEmpty())
+	{
+		return;
+	}
+
+	FScopeEnter Entry = ScopeStack.Pop();
+
+	if (Provider->NumCpuEvents)
+	{
+		++(*Provider->NumCpuEvents);
+	}
+
+	if (Provider->Options && Provider->Options->bSummaryOnly)
+	{
+		return;
+	}
+
+	double Duration = EndTime - Entry.StartTime;
+
+	if (Provider->Options)
+	{
+		if (Provider->Options->StartTime >= 0.0 && Entry.StartTime < Provider->Options->StartTime)
+		{
+			return;
+		}
+		if (Provider->Options->EndTime >= 0.0 && Entry.StartTime > Provider->Options->EndTime)
+		{
+			return;
+		}
+	}
+
+	// Resolve timer name — use the main timer ID for metadata entries
+	uint32 ResolvedTimerId = Entry.TimerId;
+	if (int32(Entry.TimerId) < 0)
+	{
+		uint32 MetaIdx = ~Entry.TimerId;
+		if (MetaIdx < uint32(Provider->Metadatas.Num()))
+		{
+			ResolvedTimerId = Provider->Metadatas[MetaIdx].TimerId;
+		}
+	}
+
+	// Apply timer name filter
+	if (Provider->Options && !Provider->Options->TimerFilter.IsEmpty())
+	{
+		const FString* TimerName = Provider->LookupTimerName(Entry.TimerId);
+		if (!TimerName || !TimerName->Contains(Provider->Options->TimerFilter))
+		{
+			return;
+		}
+	}
+
+	if (Provider->ScopeGroups)
+	{
+		FScopeGroupKey Key;
+		Key.ThreadId = ThreadId;
+		Key.SpecId = ResolvedTimerId;
+
+		FScopeDataPoint Point;
+		Point.Time = Entry.StartTime;
+		Point.Duration = Duration;
+		Provider->ScopeGroups->FindOrAdd(Key).Add(Point);
+	}
+}
+
+////////////////////////////////////////////////////////////////////////////////////////////////////
+// FTraceQueryCounter + FTraceQueryCounterProvider
+////////////////////////////////////////////////////////////////////////////////////////////////////
+
+void FTraceQueryCounter::SetValue(double Time, int64 Value)
+{
+	if (NumCounterValues)
+	{
+		++(*NumCounterValues);
+	}
+
+	if (!Options || Options->bSummaryOnly)
+	{
+		return;
+	}
+
+	if (TimeFilterCallback && !TimeFilterCallback(Time))
+	{
+		return;
+	}
+
+	if (Options && !Options->CounterFilter.IsEmpty() && !Name.Contains(Options->CounterFilter))
+	{
+		return;
+	}
+
+	if (EmitCallback && JsonEscapeCallback)
+	{
+		FString Line = FString::Printf(
+			TEXT("{\"type\":\"counter\",\"time\":%.6f,\"counterId\":%u,\"name\":\"%s\",\"value\":%lld}"),
+			Time, (uint32)CounterId, *JsonEscapeCallback(Name), Value);
+		EmitCallback(TCHAR_TO_UTF8(*Line));
+	}
+}
+
+void FTraceQueryCounter::SetValue(double Time, double Value)
+{
+	if (NumCounterValues)
+	{
+		++(*NumCounterValues);
+	}
+
+	if (!Options || Options->bSummaryOnly)
+	{
+		return;
+	}
+
+	if (TimeFilterCallback && !TimeFilterCallback(Time))
+	{
+		return;
+	}
+
+	if (Options && !Options->CounterFilter.IsEmpty() && !Name.Contains(Options->CounterFilter))
+	{
+		return;
+	}
+
+	if (EmitCallback && JsonEscapeCallback)
+	{
+		FString Line = FString::Printf(
+			TEXT("{\"type\":\"counter\",\"time\":%.6f,\"counterId\":%u,\"name\":\"%s\",\"value\":%.9g}"),
+			Time, (uint32)CounterId, *JsonEscapeCallback(Name), Value);
+		EmitCallback(TCHAR_TO_UTF8(*Line));
+	}
+}
+
+TraceServices::IEditableCounter* FTraceQueryCounterProvider::CreateEditableCounter()
+{
+	auto Counter = MakeUnique<FTraceQueryCounter>();
+	Counter->Options = Options;
+	Counter->EmitCallback = EmitCallback;
+	Counter->TimeFilterCallback = TimeFilterCallback;
+	Counter->JsonEscapeCallback = JsonEscapeCallback;
+	Counter->NumCounterValues = NumCounterValues;
+	Counter->CounterId = NextCounterId++;
+
+	if (NumCounterSpecs)
+	{
+		++(*NumCounterSpecs);
+	}
+
+	FTraceQueryCounter* Ptr = Counter.Get();
+	Counters.Add(MoveTemp(Counter));
+	return Ptr;
+}
+
+////////////////////////////////////////////////////////////////////////////////////////////////////
+// FJsonTraceAnalyzer
+////////////////////////////////////////////////////////////////////////////////////////////////////
+
+FJsonTraceAnalyzer::FJsonTraceAnalyzer(const FTraceQueryOptions& InOptions)
+	: Options(InOptions)
+{
+}
+
+////////////////////////////////////////////////////////////////////////////////////////////////////
+
+void FJsonTraceAnalyzer::OnAnalysisBegin(const FOnAnalysisContext& Context)
+{
+	FInterfaceBuilder& Builder = Context.InterfaceBuilder;
+
+	// Always subscribe to diagnostics and frame events for metadata/structure
+	Builder.RouteEvent(RouteId_DiagSession, "Diagnostics", "Session");
+	Builder.RouteEvent(RouteId_DiagSession2, "Diagnostics", "Session2");
+	Builder.RouteEvent(RouteId_BeginFrame, "Misc", "BeginFrame");
+
+	if (Options.bWantAllEvents)
+	{
+		Builder.RouteAllEvents(RouteId_AllEvents);
+	}
+
+	if (Options.bWantStats || Options.bWantAllEvents)
+	{
+		Builder.RouteEvent(RouteId_StatsSpec, "Stats", "Spec");
+		Builder.RouteEvent(RouteId_StatsEventBatch, "Stats", "EventBatch");
+		Builder.RouteEvent(RouteId_StatsEventBatch2, "Stats", "EventBatch2");
+	}
+
+	// Note: CpuProfiler and Counters routes are registered by the factory-created analyzers,
+	// not by this analyzer. See TraceQuery.cpp for the wiring.
+}
+
+void FJsonTraceAnalyzer::OnAnalysisEnd()
+{
+	// Flush buffered frame groups as grouped JSONL
+	static const char* FrameTypeNames[] = { "game", "render" };
+	for (TPair<uint8, TArray<double>>& Pair : FrameGroups)
+	{
+		uint8 FrameType = Pair.Key;
+		TArray<double>& Times = Pair.Value;
+		const char* FrameTypeName = (FrameType < UE_ARRAY_COUNT(FrameTypeNames)) ? FrameTypeNames[FrameType] : "unknown";
+
+		FString DataArray;
+		DataArray.Reserve(Times.Num() * 16);
+		DataArray += TEXT("[");
+		for (int32 Index = 0; Index < Times.Num(); ++Index)
+		{
+			if (Index > 0)
+			{
+				DataArray += TEXT(",");
+			}
+			DataArray += FString::Printf(TEXT("%.6f"), Times[Index]);
+		}
+		DataArray += TEXT("]");
+
+		FString Line = FString::Printf(
+			TEXT("{\"type\":\"frame\",\"frameType\":\"%s\",\"count\":%d,\"data\":%s}"),
+			ANSI_TO_TCHAR(FrameTypeName), Times.Num(), *DataArray);
+		EmitJsonLine(TCHAR_TO_UTF8(*Line));
+	}
+
+	// Flush buffered CPU scope groups as grouped JSONL
+	for (TPair<FScopeGroupKey, TArray<FScopeDataPoint>>& Pair : ScopeGroups)
+	{
+		const FScopeGroupKey& Key = Pair.Key;
+		TArray<FScopeDataPoint>& Points = Pair.Value;
+
+		// Resolve timer name via the CpuProfiler provider
+		FString TimerName;
+		if (CpuProfilerProvider)
+		{
+			const FString* Name = CpuProfilerProvider->LookupTimerName(Key.SpecId);
+			TimerName = Name ? *Name : FString::Printf(TEXT("Timer_%u"), Key.SpecId);
+		}
+		else
+		{
+			TimerName = FString::Printf(TEXT("Timer_%u"), Key.SpecId);
+		}
+
+		// Build data array: [[time,dur],[time,dur],...]
+		FString DataArray;
+		DataArray.Reserve(Points.Num() * 32);
+		DataArray += TEXT("[");
+		for (int32 Index = 0; Index < Points.Num(); ++Index)
+		{
+			if (Index > 0)
+			{
+				DataArray += TEXT(",");
+			}
+			DataArray += FString::Printf(TEXT("[%.6f,%.9g]"), Points[Index].Time, Points[Index].Duration);
+		}
+		DataArray += TEXT("]");
+
+		FString Line = FString::Printf(
+			TEXT("{\"type\":\"cpu_scope\",\"thread\":%u,\"specId\":%u,\"name\":\"%s\",\"count\":%d,\"data\":%s}"),
+			Key.ThreadId, Key.SpecId, *JsonEscape(TimerName), Points.Num(), *DataArray);
+		EmitJsonLine(TCHAR_TO_UTF8(*Line));
+	}
+
+	if (Options.bSummaryOnly)
+	{
+		fprintf(stdout, "=== TraceQuery Summary ===\n");
+		fprintf(stdout, "Stats specs:      %llu\n", NumStatsSpecs);
+		fprintf(stdout, "Stats batches:    %llu\n", NumStatsBatches);
+		fprintf(stdout, "Stats operations: %llu\n", NumStatsOps);
+		fprintf(stdout, "Counter specs:    %llu\n", NumCounterSpecs);
+		fprintf(stdout, "Counter values:   %llu\n", NumCounterValues);
+		fprintf(stdout, "Frames:           %llu\n", NumFrames);
+		fprintf(stdout, "CPU timer specs:  %llu\n", NumCpuSpecs);
+		fprintf(stdout, "CPU events:       %llu\n", NumCpuEvents);
+		fprintf(stdout, "Total events:     %llu\n", NumTotalEvents);
+	}
+}
+
+////////////////////////////////////////////////////////////////////////////////////////////////////
+
+bool FJsonTraceAnalyzer::OnEvent(uint16 RouteId, EStyle Style, const FOnEventContext& Context)
+{
+	++NumTotalEvents;
+
+	switch (RouteId)
+	{
+	case RouteId_StatsSpec:
+		OnStatsSpec(Context);
+		break;
+	case RouteId_StatsEventBatch:
+		OnStatsEventBatch(Context, /*bV2=*/false);
+		break;
+	case RouteId_StatsEventBatch2:
+		OnStatsEventBatch(Context, /*bV2=*/true);
+		break;
+	case RouteId_DiagSession:
+		// Session v1 uses attachment-based parsing; Session2 is cleaner.
+		// For v1, extract what we can.
+		break;
+	case RouteId_DiagSession2:
+		OnDiagnosticsSession2(Context);
+		break;
+	case RouteId_BeginFrame:
+		OnBeginFrame(Context);
+		break;
+	default:
+		break;
+	}
+
+	return true;
+}
+
+////////////////////////////////////////////////////////////////////////////////////////////////////
+// Stats
+////////////////////////////////////////////////////////////////////////////////////////////////////
+
+void FJsonTraceAnalyzer::OnStatsSpec(const FOnEventContext& Context)
+{
+	const FEventData& EventData = Context.EventData;
+	uint32 StatId = EventData.GetValue<uint32>("Id");
+
+	FStatSpec& Spec = StatSpecs.FindOrAdd(StatId);
+
+	if (!EventData.GetString("Name", Spec.Name))
+	{
+		// Legacy: name in attachment
+		const uint8* Attachment = EventData.GetAttachment();
+		if (Attachment)
+		{
+			Spec.Name = FString(reinterpret_cast<const ANSICHAR*>(Attachment));
+		}
+	}
+	EventData.GetString("Group", Spec.Group);
+	EventData.GetString("Description", Spec.Description);
+	Spec.bIsFloatingPoint = EventData.GetValue<bool>("IsFloatingPoint");
+	Spec.bIsResetEveryFrame = EventData.GetValue<bool>("ShouldClearEveryFrame");
+
+	++NumStatsSpecs;
+}
+
+void FJsonTraceAnalyzer::OnStatsEventBatch(const FOnEventContext& Context, bool bV2)
+{
+	enum class EOpType : uint8
+	{
+		Increment = 0,
+		Decrement = 1,
+		AddInteger = 2,
+		SetInteger = 3,
+		AddFloat = 4,
+		SetFloat = 5,
+	};
+
+	++NumStatsBatches;
+
+	if (Options.bSummaryOnly)
+	{
+		// Still need to count operations for summary
+		TArrayView<const uint8> DataView = TraceServices::FTraceAnalyzerUtils::LegacyAttachmentArray("Data", Context);
+		const uint8* BufferPtr = DataView.GetData();
+		const uint8* BufferEnd = BufferPtr + DataView.Num();
+		while (BufferPtr < BufferEnd)
+		{
+			uint64 DecodedIdAndOp = TraceQueryLocal::Decode7bit(BufferPtr);
+			uint8 Op = DecodedIdAndOp & 0x7;
+			TraceQueryLocal::Decode7bit(BufferPtr); // CycleDiff
+			// Skip value
+			EOpType SummaryOp = static_cast<EOpType>(Op);
+			if (SummaryOp == EOpType::AddInteger || SummaryOp == EOpType::SetInteger)
+			{
+				TraceQueryLocal::DecodeZigZag(BufferPtr);
+			}
+			else if (SummaryOp == EOpType::AddFloat || SummaryOp == EOpType::SetFloat)
+			{
+				BufferPtr += sizeof(double);
+			}
+			// Op 0 (Inc) and 1 (Dec) have no value payload
+			++NumStatsOps;
+		}
+		if (BufferPtr != BufferEnd)
+		{
+			UE_LOG(LogTraceQuery, Warning, TEXT("Stats batch buffer mismatch (summary): %lld bytes remaining"), (int64)(BufferEnd - BufferPtr));
+		}
+		return;
+	}
+
+	// Decode logic ported from StatsTraceAnalysis.cpp:189-274
+	uint32 ThreadId = TraceQueryLocal::GetThreadIdField(Context);
+	FThreadState& ThreadState = GetThreadState(ThreadId);
+
+	const uint64 BaseTimestamp = Context.EventTime.AsCycle64() - Context.EventTime.GetTimestamp();
+
+	TArrayView<const uint8> DataView = TraceServices::FTraceAnalyzerUtils::LegacyAttachmentArray("Data", Context);
+	const uint8* BufferPtr = DataView.GetData();
+	const uint8* BufferEnd = BufferPtr + DataView.Num();
+
+	while (BufferPtr < BufferEnd)
+	{
+		uint64 DecodedIdAndOp = TraceQueryLocal::Decode7bit(BufferPtr);
+		uint32 StatId = static_cast<uint32>(DecodedIdAndOp >> 3);
+		uint8 Op = DecodedIdAndOp & 0x7;
+		uint64 CycleDiff = TraceQueryLocal::Decode7bit(BufferPtr);
+
+		if (bV2)
+		{
+			if (CycleDiff >= BaseTimestamp)
+			{
+				ThreadState.LastCycle = 0;
+			}
+		}
+
+		uint64 Cycle = ThreadState.LastCycle + CycleDiff;
+		double Time = Context.EventTime.AsSeconds(Cycle);
+		ThreadState.LastCycle = Cycle;
+
+		// Decode value
+		double FloatValue = 0.0;
+		int64 IntValue = 0;
+		bool bIsFloat = false;
+
+		EOpType OpType = static_cast<EOpType>(Op);
+		switch (OpType)
+		{
+		case EOpType::Increment:
+			IntValue = 1;
+			break;
+		case EOpType::Decrement:
+			IntValue = -1;
+			break;
+		case EOpType::AddInteger:
+			IntValue = TraceQueryLocal::DecodeZigZag(BufferPtr);
+			break;
+		case EOpType::SetInteger:
+			IntValue = TraceQueryLocal::DecodeZigZag(BufferPtr);
+			break;
+		case EOpType::AddFloat:
+			bIsFloat = true;
+			memcpy(&FloatValue, BufferPtr, sizeof(double));
+			BufferPtr += sizeof(double);
+			break;
+		case EOpType::SetFloat:
+			bIsFloat = true;
+			memcpy(&FloatValue, BufferPtr, sizeof(double));
+			BufferPtr += sizeof(double);
+			break;
+		}
+
+		++NumStatsOps;
+
+		if (!PassesTimeFilter(Time))
+		{
+			continue;
+		}
+
+		// Resolve stat name
+		const FStatSpec* Spec = StatSpecs.Find(StatId);
+		FString StatName;
+		FString GroupName;
+		if (Spec)
+		{
+			StatName = Spec->Name;
+			GroupName = Spec->Group;
+		}
+		else
+		{
+			StatName = FString::Printf(TEXT("Unknown_%u"), StatId);
+		}
+
+		// Apply counter name filter (unified: matches both Stats-sourced and Counters-channel events)
+		if (!Options.CounterFilter.IsEmpty())
+		{
+			if (!StatName.Contains(Options.CounterFilter))
+			{
+				continue;
+			}
+		}
+
+		static const char* OpNames[] = { "inc", "dec", "add", "set", "add", "set" };
+		const char* OpName = (Op < UE_ARRAY_COUNT(OpNames)) ? OpNames[Op] : "unknown";
+
+		// Emit JSONL
+		if (bIsFloat)
+		{
+			FString Line = FString::Printf(
+				TEXT("{\"type\":\"stat\",\"time\":%.6f,\"thread\":%u,\"statId\":%u,\"name\":\"%s\",\"group\":\"%s\",\"op\":\"%s\",\"value\":%.9g}"),
+				Time, ThreadId, StatId, *JsonEscape(StatName), *JsonEscape(GroupName), ANSI_TO_TCHAR(OpName), FloatValue);
+			EmitJsonLine(TCHAR_TO_UTF8(*Line));
+		}
+		else
+		{
+			FString Line = FString::Printf(
+				TEXT("{\"type\":\"stat\",\"time\":%.6f,\"thread\":%u,\"statId\":%u,\"name\":\"%s\",\"group\":\"%s\",\"op\":\"%s\",\"value\":%lld}"),
+				Time, ThreadId, StatId, *JsonEscape(StatName), *JsonEscape(GroupName), ANSI_TO_TCHAR(OpName), IntValue);
+			EmitJsonLine(TCHAR_TO_UTF8(*Line));
+		}
+	}
+
+	if (BufferPtr != BufferEnd)
+	{
+		UE_LOG(LogTraceQuery, Warning, TEXT("Stats batch buffer mismatch: %lld bytes remaining"), (int64)(BufferEnd - BufferPtr));
+	}
+}
+
+////////////////////////////////////////////////////////////////////////////////////////////////////
+// Diagnostics
+////////////////////////////////////////////////////////////////////////////////////////////////////
+
+void FJsonTraceAnalyzer::OnDiagnosticsSession2(const FOnEventContext& Context)
+{
+	const FEventData& EventData = Context.EventData;
+
+	FString Platform, AppName, ProjectName, Branch, BuildVersion, CommandLine;
+	EventData.GetString("Platform", Platform);
+	EventData.GetString("AppName", AppName);
+	EventData.GetString("ProjectName", ProjectName);
+	EventData.GetString("Branch", Branch);
+	EventData.GetString("BuildVersion", BuildVersion);
+	EventData.GetString("CommandLine", CommandLine);
+	uint32 Changelist = EventData.GetValue<uint32>("Changelist", 0);
+	uint8 ConfigurationType = EventData.GetValue<uint8>("ConfigurationType");
+	uint8 TargetType = EventData.GetValue<uint8>("TargetType");
+
+	static const char* ConfigNames[] = { "Unknown", "Debug", "DebugGame", "Development", "Shipping", "Test" };
+	static const char* TargetNames[] = { "Unknown", "Game", "Server", "Client", "Editor", "Program" };
+
+	const char* ConfigName = (ConfigurationType < UE_ARRAY_COUNT(ConfigNames)) ? ConfigNames[ConfigurationType] : "Unknown";
+	const char* TargetName = (TargetType < UE_ARRAY_COUNT(TargetNames)) ? TargetNames[TargetType] : "Unknown";
+
+	FString Line = FString::Printf(
+		TEXT("{\"type\":\"session\",\"platform\":\"%s\",\"appName\":\"%s\",\"projectName\":\"%s\",\"branch\":\"%s\",\"buildVersion\":\"%s\",\"changelist\":%u,\"configuration\":\"%s\",\"targetType\":\"%s\"}"),
+		*JsonEscape(Platform), *JsonEscape(AppName), *JsonEscape(ProjectName),
+		*JsonEscape(Branch), *JsonEscape(BuildVersion), Changelist,
+		ANSI_TO_TCHAR(ConfigName), ANSI_TO_TCHAR(TargetName));
+	EmitJsonLine(TCHAR_TO_UTF8(*Line));
+}
+
+////////////////////////////////////////////////////////////////////////////////////////////////////
+// Frame
+////////////////////////////////////////////////////////////////////////////////////////////////////
+
+void FJsonTraceAnalyzer::OnBeginFrame(const FOnEventContext& Context)
+{
+	const FEventData& EventData = Context.EventData;
+	uint8 FrameType = EventData.GetValue<uint8>("FrameType");
+	uint64 Cycle = EventData.GetValue<uint64>("Cycle");
+	double Time = Context.EventTime.AsSeconds(Cycle);
+
+	++NumFrames;
+
+	if (Options.bSummaryOnly || !PassesTimeFilter(Time))
+	{
+		return;
+	}
+
+	FrameGroups.FindOrAdd(FrameType).Add(Time);
+}
+
+////////////////////////////////////////////////////////////////////////////////////////////////////
+// Helpers
+////////////////////////////////////////////////////////////////////////////////////////////////////
+
+bool FJsonTraceAnalyzer::PassesTimeFilter(double Time) const
+{
+	if (Options.StartTime >= 0.0 && Time < Options.StartTime)
+	{
+		return false;
+	}
+	if (Options.EndTime >= 0.0 && Time > Options.EndTime)
+	{
+		return false;
+	}
+	return true;
+}
+
+void FJsonTraceAnalyzer::EmitJsonLine(const char* JsonLine)
+{
+	if (!Options.bSummaryOnly)
+	{
+		puts(JsonLine);
+	}
+}
+
+FString FJsonTraceAnalyzer::JsonEscape(const FString& Input)
+{
+	FString Result;
+	Result.Reserve(Input.Len() + 8);
+
+	for (int32 Index = 0; Index < Input.Len(); ++Index)
+	{
+		TCHAR Ch = Input[Index];
+		switch (Ch)
+		{
+		case TEXT('"'):  Result += TEXT("\\\""); break;
+		case TEXT('\\'): Result += TEXT("\\\\"); break;
+		case TEXT('\n'): Result += TEXT("\\n");  break;
+		case TEXT('\r'): Result += TEXT("\\r");  break;
+		case TEXT('\t'): Result += TEXT("\\t");  break;
+		default:
+			if (Ch < 0x20)
+			{
+				Result += FString::Printf(TEXT("\\u%04x"), (uint32)Ch);
+			}
+			else
+			{
+				Result.AppendChar(Ch);
+			}
+			break;
+		}
+	}
+
+	return Result;
+}
+
+FJsonTraceAnalyzer::FThreadState& FJsonTraceAnalyzer::GetThreadState(uint32 ThreadId)
+{
+	TSharedPtr<FThreadState>& Slot = ThreadStates.FindOrAdd(ThreadId);
+	if (!Slot.IsValid())
+	{
+		Slot = MakeShared<FThreadState>();
+	}
+	return *Slot;
+}
diff --git a/Engine/Source/Programs/TraceQuery/Private/JsonTraceAnalyzer.h b/Engine/Source/Programs/TraceQuery/Private/JsonTraceAnalyzer.h
new file mode 100644
--- /dev/null
+++ b/Engine/Source/Programs/TraceQuery/Private/JsonTraceAnalyzer.h
@@ -0,0 +1,257 @@
+// Copyright Epic Games, Inc. All Rights Reserved.
+
+#pragma once
+
+#include "TraceQuery.h"
+
+#include "Containers/Map.h"
+#include "Containers/UnrealString.h"
+#include "Misc/Optional.h"
+#include "Templates/SharedPointer.h"
+#include "Templates/UniquePtr.h"
+#include "Trace/Analyzer.h"
+#include "TraceServices/Model/Counters.h"
+#include "TraceServices/Model/Threads.h"
+#include "TraceServices/Model/TimingProfiler.h"
+
+////////////////////////////////////////////////////////////////////////////////////////////////////
+// Shared data types used by providers and JsonTraceAnalyzer
+////////////////////////////////////////////////////////////////////////////////////////////////////
+
+// Key for grouping CPU scope events by (ThreadId, SpecId)
+struct FScopeGroupKey
+{
+	uint32 ThreadId;
+	uint32 SpecId;
+
+	bool operator==(const FScopeGroupKey& Other) const
+	{
+		return ThreadId == Other.ThreadId && SpecId == Other.SpecId;
+	}
+};
+
+inline uint32 GetTypeHash(const FScopeGroupKey& Key)
+{
+	return HashCombine(GetTypeHash(Key.ThreadId), GetTypeHash(Key.SpecId));
+}
+
+struct FScopeDataPoint
+{
+	double Time = 0.0;
+	double Duration = 0.0;
+};
+
+////////////////////////////////////////////////////////////////////////////////////////////////////
+// FTraceQueryCpuProfilerProvider
+//
+// Lightweight IEditableTimingProfilerProvider + IEditableThreadProvider that receives decoded
+// CPU profiler events from TraceServices' FCpuProfilerAnalyzer and populates ScopeGroups.
+// Modeled after FSummarizeCpuProfilerProvider in SummarizeTraceUtils.h.
+////////////////////////////////////////////////////////////////////////////////////////////////////
+
+class FTraceQueryCpuProfilerProvider
+	: public TraceServices::IEditableThreadProvider
+	, public TraceServices::IEditableTimingProfilerProvider
+{
+public:
+	// Wired by TraceQuery.cpp before analysis begins
+	const FTraceQueryOptions* Options = nullptr;
+	TMap<FScopeGroupKey, TArray<FScopeDataPoint>>* ScopeGroups = nullptr;
+	uint64* NumCpuSpecs = nullptr;
+	uint64* NumCpuEvents = nullptr;
+
+	// Resolve a TimerIndex (which may be a ~MetadataId) to its scope name, or nullptr
+	const FString* LookupTimerName(uint32 TimerId) const;
+
+private:
+	struct FScopeEnter
+	{
+		uint32 TimerId;
+		double StartTime;
+	};
+
+	struct FThread : public TraceServices::IEditableTimeline<TraceServices::FTimingProfilerEvent>
+	{
+		FThread(uint32 InThreadId, FTraceQueryCpuProfilerProvider* InProvider)
+			: ThreadId(InThreadId)
+			, Provider(InProvider)
+		{
+		}
+
+		virtual void AppendBeginEvent(double StartTime, const TraceServices::FTimingProfilerEvent& Event) override;
+		virtual void AppendEndEvent(double EndTime) override;
+
+		uint32 ThreadId;
+		FTraceQueryCpuProfilerProvider* Provider;
+		TArray<FScopeEnter> ScopeStack;
+	};
+
+	struct FMetadata
+	{
+		TArray<uint8> Payload;
+		uint32 TimerId;
+	};
+
+	// IEditableThreadProvider
+	virtual void AddThread(uint32 Id, const TCHAR* Name, EThreadPriority Priority) override;
+
+	// IEditableTimingProfilerProvider — methods called by FCpuProfilerAnalyzer
+	virtual uint32 AddTimer(TraceServices::ETimingProfilerTimerType Type) override;
+	virtual uint32 AddCpuTimer(FStringView Name, const TCHAR* File, uint32 Line) override;
+	virtual void SetTimerName(uint32 TimerId, FStringView Name) override;
+	virtual uint32 AddMetadata(uint32 MainTimerId, TArray<uint8>&& Metadata) override;
+	virtual TArrayView<uint8> GetEditableMetadata(uint32 TimerId) override;
+	virtual TraceServices::IEditableTimeline<TraceServices::FTimingProfilerEvent>& GetCpuThreadEditableTimeline(uint32 ThreadId) override;
+
+	uint32 AddTimerInternal(FStringView Name);
+
+	TArray<TOptional<FString>> TimerNames;
+	TMap<uint32, TUniquePtr<FThread>> Threads;
+	TArray<FMetadata> Metadatas;
+};
+
+////////////////////////////////////////////////////////////////////////////////////////////////////
+// FTraceQueryCounter + FTraceQueryCounterProvider
+//
+// Lightweight IEditableCounter / IEditableCounterProvider that streams counter values
+// as JSONL via a callback. Modeled after FSummarizeCounter in SummarizeTraceCommandlet.cpp.
+////////////////////////////////////////////////////////////////////////////////////////////////////
+
+class FTraceQueryCounter : public TraceServices::IEditableCounter
+{
+public:
+	// Wired by provider on creation
+	const FTraceQueryOptions* Options = nullptr;
+	TFunction<void(const char*)> EmitCallback;
+	TFunction<bool(double)> TimeFilterCallback;
+	TFunction<FString(const FString&)> JsonEscapeCallback;
+	uint64* NumCounterValues = nullptr;
+
+	uint16 CounterId = 0;
+
+	// IEditableCounter
+	virtual void SetName(const TCHAR* InName) override { Name = InName; }
+	virtual void SetGroup(const TCHAR* Group) override {}
+	virtual void SetDescription(const TCHAR* Description) override {}
+	virtual void SetIsFloatingPoint(bool bIsFloatingPoint) override { bIsFloat = bIsFloatingPoint; }
+	virtual void SetIsResetEveryFrame(bool bInIsResetEveryFrame) override {}
+	virtual void SetDisplayHint(TraceServices::ECounterDisplayHint DisplayHint) override {}
+
+	virtual void AddValue(double Time, int64 Value) override { SetValue(Time, Value); }
+	virtual void SetValue(double Time, int64 Value) override;
+	virtual void AddValue(double Time, double Value) override { SetValue(Time, Value); }
+	virtual void SetValue(double Time, double Value) override;
+
+private:
+	FString Name;
+	bool bIsFloat = false;
+};
+
+class FTraceQueryCounterProvider : public TraceServices::IEditableCounterProvider
+{
+public:
+	// Wired by TraceQuery.cpp before analysis begins
+	const FTraceQueryOptions* Options = nullptr;
+	TFunction<void(const char*)> EmitCallback;
+	TFunction<bool(double)> TimeFilterCallback;
+	TFunction<FString(const FString&)> JsonEscapeCallback;
+	uint64* NumCounterSpecs = nullptr;
+	uint64* NumCounterValues = nullptr;
+
+	// IEditableCounterProvider
+	virtual const TraceServices::ICounter* GetCounter(TraceServices::IEditableCounter* EditableCounter) override { return nullptr; }
+	virtual TraceServices::IEditableCounter* CreateEditableCounter() override;
+	virtual void AddCounter(const TraceServices::ICounter* Counter) override {}
+
+private:
+	uint16 NextCounterId = 0;
+	TArray<TUniquePtr<FTraceQueryCounter>> Counters;
+};
+
+////////////////////////////////////////////////////////////////////////////////////////////////////
+// FJsonTraceAnalyzer
+//
+// IAnalyzer that handles Stats, Diagnostics, BeginFrame, and AllEvents.
+// CpuProfiler and Counters are handled by TraceServices factory analyzers with custom providers.
+////////////////////////////////////////////////////////////////////////////////////////////////////
+
+class FJsonTraceAnalyzer : public UE::Trace::IAnalyzer
+{
+public:
+	explicit FJsonTraceAnalyzer(const FTraceQueryOptions& InOptions);
+
+	// Escape a string for JSON output (handles quotes, backslashes, control chars)
+	static FString JsonEscape(const FString& Input);
+
+	// Output helpers — public so providers can access them
+	bool PassesTimeFilter(double Time) const;
+	void EmitJsonLine(const char* JsonLine);
+
+	// Buffered CPU scope data — populated by CpuProfiler provider, flushed in OnAnalysisEnd
+	TMap<FScopeGroupKey, TArray<FScopeDataPoint>> ScopeGroups;
+
+	// Summary mode counters — providers write to these via pointers
+	uint64 NumStatsSpecs = 0;
+	uint64 NumStatsBatches = 0;
+	uint64 NumStatsOps = 0;
+	uint64 NumCounterSpecs = 0;
+	uint64 NumCounterValues = 0;
+	uint64 NumFrames = 0;
+	uint64 NumCpuSpecs = 0;
+	uint64 NumCpuEvents = 0;
+	uint64 NumTotalEvents = 0;
+
+	// Set by TraceQuery.cpp to enable timer name resolution in OnAnalysisEnd
+	const FTraceQueryCpuProfilerProvider* CpuProfilerProvider = nullptr;
+
+private:
+	// IAnalyzer interface
+	virtual void OnAnalysisBegin(const FOnAnalysisContext& Context) override;
+	virtual void OnAnalysisEnd() override;
+	virtual bool OnEvent(uint16 RouteId, EStyle Style, const FOnEventContext& Context) override;
+
+	// Stats decode (kept — no factory exists for Stats)
+	void OnStatsSpec(const FOnEventContext& Context);
+	void OnStatsEventBatch(const FOnEventContext& Context, bool bV2);
+
+	// Simple event handlers
+	void OnDiagnosticsSession2(const FOnEventContext& Context);
+	void OnBeginFrame(const FOnEventContext& Context);
+
+	enum : uint16
+	{
+		RouteId_StatsSpec,
+		RouteId_StatsEventBatch,
+		RouteId_StatsEventBatch2,
+		RouteId_DiagSession,
+		RouteId_DiagSession2,
+		RouteId_BeginFrame,
+		RouteId_AllEvents,
+	};
+
+	struct FStatSpec
+	{
+		FString Name;
+		FString Group;
+		FString Description;
+		bool bIsFloatingPoint = false;
+		bool bIsResetEveryFrame = false;
+	};
+
+	// Per-thread state for Stats batch decode
+	struct FThreadState
+	{
+		uint64 LastCycle = 0;
+	};
+
+	const FTraceQueryOptions& Options;
+
+	TMap<uint32, FStatSpec> StatSpecs;
+	TMap<uint32, TSharedPtr<FThreadState>> ThreadStates;
+
+	// Buffered frame timestamps — flushed as grouped JSONL in OnAnalysisEnd
+	// Key: frameType (0=game, 1=render)
+	TMap<uint8, TArray<double>> FrameGroups;
+
+	FThreadState& GetThreadState(uint32 ThreadId);
+};
diff --git a/Engine/Source/Programs/TraceQuery/Private/TraceQuery.cpp b/Engine/Source/Programs/TraceQuery/Private/TraceQuery.cpp
new file mode 100644
--- /dev/null
+++ b/Engine/Source/Programs/TraceQuery/Private/TraceQuery.cpp
@@ -0,0 +1,164 @@
+// Copyright Epic Games, Inc. All Rights Reserved.
+
+#include "TraceQuery.h"
+#include "JsonTraceAnalyzer.h"
+
+#include "Misc/CommandLine.h"
+#include "RequiredProgramMainCPPInclude.h"
+
+// TraceAnalysis
+#include "Trace/Analysis.h"
+#include "Trace/DataStream.h"
+
+// TraceServices — factory analyzers for CpuProfiler and Counters
+#include "TraceServices/AnalyzerFactories.h"
+
+#include <stdio.h>
+
+DEFINE_LOG_CATEGORY(LogTraceQuery);
+
+IMPLEMENT_APPLICATION(TraceQuery, "TraceQuery");
+
+////////////////////////////////////////////////////////////////////////////////////////////////////
+
+static void PrintUsage()
+{
+	fprintf(stderr, "Usage:\n");
+	fprintf(stderr, "  TraceQuery <file.utrace> [options]\n");
+	fprintf(stderr, "\n");
+	fprintf(stderr, "Options:\n");
+	fprintf(stderr, "  -logger=<name>        Filter by logger (Stats, Counters, CpuProfiler, or *)\n");
+	fprintf(stderr, "  -event=<pattern>      Filter by event name (* wildcard)\n");
+	fprintf(stderr, "  -counter=<name>       Filter counters by name substring (both Stats and Counters channels)\n");
+	fprintf(stderr, "  -timers=<name>        Filter CPU profiler scopes by name substring\n");
+	fprintf(stderr, "  -summary              Print event type summary, exit\n");
+	fprintf(stderr, "  -startTime=<sec>      Only emit events after this time\n");
+	fprintf(stderr, "  -endTime=<sec>        Only emit events before this time\n");
+	fprintf(stderr, "\n");
+	fprintf(stderr, "Output: JSONL on stdout (one JSON object per line)\n");
+}
+
+static bool ParseArgs(int32 ArgC, TCHAR const* const* ArgV, FTraceQueryOptions& OutOptions)
+{
+	if (ArgC < 2)
+	{
+		PrintUsage();
+		return false;
+	}
+
+	OutOptions.InputFile = ArgV[1];
+
+	for (int32 ArgIndex = 2; ArgIndex < ArgC; ++ArgIndex)
+	{
+		FString Arg = ArgV[ArgIndex];
+
+		if (Arg.StartsWith(TEXT("-logger=")))
+		{
+			OutOptions.Logger = Arg.Mid(8);
+		}
+		else if (Arg.StartsWith(TEXT("-event=")))
+		{
+			OutOptions.EventPattern = Arg.Mid(7);
+		}
+		else if (Arg.StartsWith(TEXT("-counter=")))
+		{
+			OutOptions.CounterFilter = Arg.Mid(9);
+		}
+		else if (Arg.StartsWith(TEXT("-timers=")))
+		{
+			OutOptions.TimerFilter = Arg.Mid(8);
+		}
+		else if (Arg == TEXT("-summary"))
+		{
+			OutOptions.bSummaryOnly = true;
+		}
+		else if (Arg.StartsWith(TEXT("-startTime=")))
+		{
+			OutOptions.StartTime = FCString::Atod(*Arg.Mid(11));
+		}
+		else if (Arg.StartsWith(TEXT("-endTime=")))
+		{
+			OutOptions.EndTime = FCString::Atod(*Arg.Mid(9));
+		}
+		else
+		{
+			fprintf(stderr, "Warning: Unknown argument '%s'\n", TCHAR_TO_ANSI(*Arg));
+		}
+	}
+
+	OutOptions.ResolveLoggerFlags();
+	return true;
+}
+
+////////////////////////////////////////////////////////////////////////////////////////////////////
+
+static int32 TraceQueryMain(int32 ArgC, TCHAR const* const* ArgV)
+{
+	FTraceQueryOptions Options;
+	if (!ParseArgs(ArgC, ArgV, Options))
+	{
+		return 1;
+	}
+
+	UE::Trace::FFileDataStream DataStream;
+	if (!DataStream.Open(*Options.InputFile))
+	{
+		fprintf(stderr, "Error: Cannot open trace file '%s'\n", TCHAR_TO_ANSI(*Options.InputFile));
+		return 2;
+	}
+
+	{
+		FJsonTraceAnalyzer Analyzer(Options);
+
+		UE::Trace::FAnalysisContext AnalysisContext;
+		AnalysisContext.AddAnalyzer(Analyzer);
+
+		// Create an analysis session (needed by TraceServices factory analyzers for string
+		// interning, edit scopes, etc.). We drive analysis ourselves via FAnalysisContext::Process,
+		// so the session doesn't need a data stream.
+		TSharedPtr<TraceServices::IAnalysisSession> Session =
+			TraceServices::CreateAnalysisSession(0, nullptr, TUniquePtr<UE::Trace::IInDataStream>(nullptr));
+
+		// CPU profiler via TraceServices factory analyzer
+		TSharedPtr<UE::Trace::IAnalyzer> CpuProfilerAnalyzer;
+		FTraceQueryCpuProfilerProvider CpuProfilerProvider;
+		if (Options.bWantCpuProfiler || Options.bWantAllEvents)
+		{
+			CpuProfilerProvider.Options = &Options;
+			CpuProfilerProvider.ScopeGroups = &Analyzer.ScopeGroups;
+			CpuProfilerProvider.NumCpuSpecs = &Analyzer.NumCpuSpecs;
+			CpuProfilerProvider.NumCpuEvents = &Analyzer.NumCpuEvents;
+			Analyzer.CpuProfilerProvider = &CpuProfilerProvider;
+			CpuProfilerAnalyzer = TraceServices::CreateCpuProfilerAnalyzer(
+				*Session, CpuProfilerProvider, CpuProfilerProvider);
+			AnalysisContext.AddAnalyzer(*CpuProfilerAnalyzer);
+		}
+
+		// Counters via TraceServices factory analyzer
+		TSharedPtr<UE::Trace::IAnalyzer> CountersAnalyzer;
+		FTraceQueryCounterProvider CountersProvider;
+		if (Options.bWantCounters || Options.bWantAllEvents)
+		{
+			CountersProvider.Options = &Options;
+			CountersProvider.EmitCallback = [&Analyzer](const char* Line) { Analyzer.EmitJsonLine(Line); };
+			CountersProvider.TimeFilterCallback = [&Analyzer](double T) { return Analyzer.PassesTimeFilter(T); };
+			CountersProvider.JsonEscapeCallback = [](const FString& S) { return FJsonTraceAnalyzer::JsonEscape(S); };
+			CountersProvider.NumCounterSpecs = &Analyzer.NumCounterSpecs;
+			CountersProvider.NumCounterValues = &Analyzer.NumCounterValues;
+			CountersAnalyzer = TraceServices::CreateCountersAnalyzer(*Session, CountersProvider);
+			AnalysisContext.AddAnalyzer(*CountersAnalyzer);
+		}
+
+		AnalysisContext.Process(DataStream).Wait();
+	}
+
+	return 0;
+}
+
+////////////////////////////////////////////////////////////////////////////////////////////////////
+
+INT32_MAIN_INT32_ARGC_TCHAR_ARGV()
+{
+	FCommandLine::Set(TEXT(""));
+	return TraceQueryMain(ArgC, ArgV);
+}
diff --git a/Engine/Source/Programs/TraceQuery/Private/TraceQuery.h b/Engine/Source/Programs/TraceQuery/Private/TraceQuery.h
new file mode 100644
--- /dev/null
+++ b/Engine/Source/Programs/TraceQuery/Private/TraceQuery.h
@@ -0,0 +1,76 @@
+// Copyright Epic Games, Inc. All Rights Reserved.
+
+#pragma once
+
+#include "CoreTypes.h"
+#include "Containers/UnrealString.h"
+#include "Logging/LogCategory.h"
+
+DECLARE_LOG_CATEGORY_EXTERN(LogTraceQuery, Log, All);
+
+/**
+ * Options parsed from the command line that control which trace events are decoded
+ * and how the JSONL output is filtered and formatted.
+ */
+struct FTraceQueryOptions
+{
+	FString InputFile;
+
+	// Filter by logger name (e.g., "Stats", "Counters", "CpuProfiler")
+	FString Logger;
+
+	// Filter by event name pattern (* wildcard)
+	FString EventPattern;
+
+	// Filter counters by name substring (matches both Stats-sourced and Counters-channel events)
+	FString CounterFilter;
+
+	// Filter CPU profiler scopes by name substring
+	FString TimerFilter;
+
+	// Only emit events after this time (seconds)
+	double StartTime = -1.0;
+
+	// Only emit events before this time (seconds)
+	double EndTime = -1.0;
+
+	// Print event type summary to stdout, then exit
+	bool bSummaryOnly = false;
+
+	// Whether to include specific loggers (derived from -logger and filter flags)
+	bool bWantStats = true;
+	bool bWantCounters = true;
+	bool bWantCpuProfiler = true;  // Enabled by default — cycle stats live in CpuProfiler
+	bool bWantAllEvents = false;
+
+	void ResolveLoggerFlags()
+	{
+		if (!Logger.IsEmpty())
+		{
+			// If a specific logger is requested, only enable that one
+			bWantStats = Logger.Equals(TEXT("Stats"), ESearchCase::IgnoreCase);
+			bWantCounters = Logger.Equals(TEXT("Counters"), ESearchCase::IgnoreCase);
+			bWantCpuProfiler = Logger.Equals(TEXT("CpuProfiler"), ESearchCase::IgnoreCase);
+			bWantAllEvents = Logger.Equals(TEXT("*"), ESearchCase::IgnoreCase);
+
+			// -logger=Stats should also enable CpuProfiler since cycle stats are there
+			if (bWantStats)
+			{
+				bWantCpuProfiler = true;
+			}
+		}
+
+		// -counter filter implies we want both Stats batches AND Counters channel
+		if (!CounterFilter.IsEmpty())
+		{
+			bWantStats = true;
+			bWantCounters = true;
+		}
+
+		// -timers filter implies we want CpuProfiler
+		if (!TimerFilter.IsEmpty())
+		{
+			bWantCpuProfiler = true;
+		}
+	}
+};
diff --git a/Engine/Source/Programs/TraceQuery/TraceQuery.Build.cs b/Engine/Source/Programs/TraceQuery/TraceQuery.Build.cs
new file mode 100644
--- /dev/null
+++ b/Engine/Source/Programs/TraceQuery/TraceQuery.Build.cs
@@ -0,0 +1,15 @@
+// Copyright Epic Games, Inc. All Rights Reserved.
+
+using UnrealBuildTool;
+
+public class TraceQuery : ModuleRules
+{
+	public TraceQuery(ReadOnlyTargetRules Target) : base(Target)
+	{
+		PublicIncludePathModuleNames.Add("Launch");
+		PrivateDependencyModuleNames.Add("Core");
+		PrivateDependencyModuleNames.Add("Projects");
+		PrivateDependencyModuleNames.Add("TraceAnalysis");
+		PrivateDependencyModuleNames.Add("TraceServices");
+	}
+}
diff --git a/Engine/Source/Programs/TraceQuery/TraceQuery.Target.cs b/Engine/Source/Programs/TraceQuery/TraceQuery.Target.cs
new file mode 100644
--- /dev/null
+++ b/Engine/Source/Programs/TraceQuery/TraceQuery.Target.cs
@@ -0,0 +1,51 @@
+// Copyright Epic Games, Inc. All Rights Reserved.
+
+using UnrealBuildTool;
+using System.Collections.Generic;
+
+[SupportedPlatforms(UnrealPlatformClass.All)]
+public class TraceQueryTarget : TargetRules
+{
+	public TraceQueryTarget(TargetInfo Target) : base(Target)
+	{
+		Type = TargetType.Program;
+		LaunchModuleName = "TraceQuery";
+		LinkType = TargetLinkType.Monolithic;
+
+		IncludeOrderVersion = EngineIncludeOrderVersion.Latest;
+
+		// This app is a console application, not a Windows app (sets entry point to main(), instead of WinMain())
+		bIsBuildingConsoleApplication = true;
+
+		// Lean and mean
+		bBuildDeveloperTools = false;
+
+		bBuildWithEditorOnlyData = false;
+
+		// Currently this app is not linking against the engine, so we'll compile out references from Core to the rest of the engine
+		bCompileAgainstEngine = false;
+		bCompileAgainstCoreUObject = false;
+		bCompileAgainstApplicationCore = false;
+		bCompileICU = false;
+
+		bEnableTrace = false; // UE_TRACE_ENABLED
+		GlobalDefinitions.Add("CPUPROFILERTRACE_ENABLED=0");
+		GlobalDefinitions.Add("GPUPROFILERTRACE_ENABLED=0");
+		GlobalDefinitions.Add("LOADTIMEPROFILERTRACE_ENABLED=0");
+		GlobalDefinitions.Add("STATSTRACE_ENABLED=0");
+		GlobalDefinitions.Add("COUNTERSTRACE_ENABLED=0");
+		GlobalDefinitions.Add("LOGTRACE_ENABLED=0");
+		GlobalDefinitions.Add("MISCTRACE_ENABLED=0");
+		GlobalDefinitions.Add("UE_TASK_TRACE_ENABLED=0");
+		GlobalDefinitions.Add("UE_NET_TRACE_ENABLED=0");
+		GlobalDefinitions.Add("UE_MEMORY_TAGS_TRACE_ENABLED=0");
+		GlobalDefinitions.Add("UE_MEMORY_TRACE_ENABLED=0");
+		GlobalDefinitions.Add("UE_CALLSTACK_TRACE_ENABLED=0");
+
+		// see Engine\Source\Developer\TraceAnalysis\Public\TraceAnalysisDebug.h
+		GlobalDefinitions.Add("UE_TRACE_ANALYSIS_DEBUG_API=0");
+		GlobalDefinitions.Add("UE_TRACE_ANALYSIS_DEBUG=0");
+		GlobalDefinitions.Add("UE_TRACE_ANALYSIS_DEBUG_LOG_IMPL=1");
+		GlobalDefinitions.Add("UE_TRACE_ANALYSIS_DEBUG_LEVEL=2");
+	}
+}
