[TraceQuery] Deadlock fix on exit and minor improvements
* Added RequestEngineExit to avoid deadlock on exit
* Add Depth and ParentSpecId to the ScopeDataPoint for hierarchical info

#Synced-CL 52134467

diff --git a/Engine/Source/Programs/TraceQuery/Private/JsonTraceAnalyzer.cpp b/Engine/Source/Programs/TraceQuery/Private/JsonTraceAnalyzer.cpp
--- a/Engine/Source/Programs/TraceQuery/Private/JsonTraceAnalyzer.cpp
+++ b/Engine/Source/Programs/TraceQuery/Private/JsonTraceAnalyzer.cpp
@@ -227,9 +227,30 @@ void FTraceQueryCpuProfilerProvider::FThread::AppendEndEvent(double EndTime)
 		Key.ThreadId = ThreadId;
 		Key.SpecId = ResolvedTimerId;
 
+		// Resolve parent from the scope stack (after pop, top is the parent)
+		uint32 ParentSpecId = 0;
+		if (!ScopeStack.IsEmpty())
+		{
+			uint32 ParentTimerId = ScopeStack.Last().TimerId;
+			if (int32(ParentTimerId) < 0)
+			{
+				uint32 MetaIdx = ~ParentTimerId;
+				if (MetaIdx < uint32(Provider->Metadatas.Num()))
+				{
+					ParentSpecId = Provider->Metadatas[MetaIdx].TimerId;
+				}
+			}
+			else
+			{
+				ParentSpecId = ParentTimerId;
+			}
+		}
+
 		FScopeDataPoint Point;
 		Point.Time = Entry.StartTime;
 		Point.Duration = Duration;
+		Point.Depth = ScopeStack.Num();
+		Point.ParentSpecId = ParentSpecId;
 		Provider->ScopeGroups->FindOrAdd(Key).Add(Point);
 	}
 }
@@ -403,9 +424,10 @@ void FJsonTraceAnalyzer::OnAnalysisEnd()
 			TimerName = FString::Printf(TEXT("Timer_%u"), Key.SpecId);
 		}
 
-		// Build data array: [[time,dur],[time,dur],...]
+		// Resolve parent names for this group's data points
+		// Build data array: [[time,dur,depth,parentName],...]
 		FString DataArray;
-		DataArray.Reserve(Points.Num() * 32);
+		DataArray.Reserve(Points.Num() * 48);
 		DataArray += TEXT("[");
 		for (int32 Index = 0; Index < Points.Num(); ++Index)
 		{
@@ -413,7 +435,28 @@ void FJsonTraceAnalyzer::OnAnalysisEnd()
 			{
 				DataArray += TEXT(",");
 			}
-			DataArray += FString::Printf(TEXT("[%.6f,%.9g]"), Points[Index].Time, Points[Index].Duration);
+
+			FString ParentName;
+			if (Points[Index].ParentSpecId != 0 && CpuProfilerProvider)
+			{
+				const FString* Name = CpuProfilerProvider->LookupTimerName(Points[Index].ParentSpecId);
+				if (Name)
+				{
+					ParentName = *Name;
+				}
+			}
+
+			if (ParentName.IsEmpty())
+			{
+				DataArray += FString::Printf(TEXT("[%.6f,%.9g,%d,null]"),
+					Points[Index].Time, Points[Index].Duration, Points[Index].Depth);
+			}
+			else
+			{
+				DataArray += FString::Printf(TEXT("[%.6f,%.9g,%d,\"%s\"]"),
+					Points[Index].Time, Points[Index].Duration, Points[Index].Depth,
+					*JsonEscape(ParentName));
+			}
 		}
 		DataArray += TEXT("]");
 
diff --git a/Engine/Source/Programs/TraceQuery/Private/JsonTraceAnalyzer.h b/Engine/Source/Programs/TraceQuery/Private/JsonTraceAnalyzer.h
--- a/Engine/Source/Programs/TraceQuery/Private/JsonTraceAnalyzer.h
+++ b/Engine/Source/Programs/TraceQuery/Private/JsonTraceAnalyzer.h
@@ -39,6 +39,8 @@ struct FScopeDataPoint
 {
 	double Time = 0.0;
 	double Duration = 0.0;
+	int32 Depth = 0;
+	uint32 ParentSpecId = 0;
 };
 
 ////////////////////////////////////////////////////////////////////////////////////////////////////
diff --git a/Engine/Source/Programs/TraceQuery/Private/TraceQuery.cpp b/Engine/Source/Programs/TraceQuery/Private/TraceQuery.cpp
--- a/Engine/Source/Programs/TraceQuery/Private/TraceQuery.cpp
+++ b/Engine/Source/Programs/TraceQuery/Private/TraceQuery.cpp
@@ -152,6 +152,7 @@ static int32 TraceQueryMain(int32 ArgC, TCHAR const* const* ArgV)
 		AnalysisContext.Process(DataStream).Wait();
 	}
 
+	RequestEngineExit(TEXT("TraceQuery completed"));
 	return 0;
 }
 
