[TraceQuery] Add ThreadNames to
* Store thread names in FTraceQueryCpuProfilerProvider, fallback to "Thread_X" if no name is available

#Synced-CL 52532922

diff --git a/Engine/Source/Programs/TraceQuery/Private/JsonTraceAnalyzer.cpp b/Engine/Source/Programs/TraceQuery/Private/JsonTraceAnalyzer.cpp
--- a/Engine/Source/Programs/TraceQuery/Private/JsonTraceAnalyzer.cpp
+++ b/Engine/Source/Programs/TraceQuery/Private/JsonTraceAnalyzer.cpp
@@ -79,9 +79,20 @@ const FString* FTraceQueryCpuProfilerProvider::LookupTimerName(uint32 TimerId) c
 	return nullptr;
 }
 
+const FString* FTraceQueryCpuProfilerProvider::LookupThreadName(uint32 ThreadId) const
+{
+	return ThreadNames.Find(ThreadId);
+}
+
+void FTraceQueryCpuProfilerProvider::SetThreadName(uint32 ThreadId, const FString& Name)
+{
+	ThreadNames.Add(ThreadId, Name);
+}
+
 void FTraceQueryCpuProfilerProvider::AddThread(uint32 Id, const TCHAR* Name, EThreadPriority Priority)
 {
-	// Ensure thread exists in our map
+	// Note: Name is always nullptr here — the CpuProfiler analyzer calls AddThread for fault
+	// tolerance but doesn't have thread names. Real names arrive via FJsonTraceAnalyzer::OnThreadInfo.
 	if (!Threads.Contains(Id))
 	{
 		Threads.Add(Id, MakeUnique<FThread>(Id, this));
@@ -377,6 +388,20 @@ void FJsonTraceAnalyzer::OnAnalysisBegin(const FOnAnalysisContext& Context)
 	// not by this analyzer. See TraceQuery.cpp for the wiring.
 }
 
+void FJsonTraceAnalyzer::OnThreadInfo(const FThreadInfo& ThreadInfo)
+{
+	// Feed thread names from the trace's Diagnostics.Thread events into our CpuProfiler provider
+	if (CpuProfilerProvider)
+	{
+		uint32 ThreadId = ThreadInfo.GetId();
+		FString Name = ThreadInfo.GetName();
+		if (!Name.IsEmpty())
+		{
+			CpuProfilerProvider->SetThreadName(ThreadId, Name);
+		}
+	}
+}
+
 void FJsonTraceAnalyzer::OnAnalysisEnd()
 {
 	// Flush buffered frame groups as grouped JSONL
@@ -460,9 +485,21 @@ void FJsonTraceAnalyzer::OnAnalysisEnd()
 		}
 		DataArray += TEXT("]");
 
+		// Resolve thread name via the CpuProfiler provider
+		FString ThreadName;
+		if (CpuProfilerProvider)
+		{
+			const FString* Name = CpuProfilerProvider->LookupThreadName(Key.ThreadId);
+			ThreadName = Name ? *Name : FString::Printf(TEXT("Thread_%u"), Key.ThreadId);
+		}
+		else
+		{
+			ThreadName = FString::Printf(TEXT("Thread_%u"), Key.ThreadId);
+		}
+
 		FString Line = FString::Printf(
-			TEXT("{\"type\":\"cpu_scope\",\"thread\":%u,\"specId\":%u,\"name\":\"%s\",\"count\":%d,\"data\":%s}"),
-			Key.ThreadId, Key.SpecId, *JsonEscape(TimerName), Points.Num(), *DataArray);
+			TEXT("{\"type\":\"cpu_scope\",\"thread\":%u,\"threadName\":\"%s\",\"specId\":%u,\"name\":\"%s\",\"count\":%d,\"data\":%s}"),
+			Key.ThreadId, *JsonEscape(ThreadName), Key.SpecId, *JsonEscape(TimerName), Points.Num(), *DataArray);
 		EmitJsonLine(TCHAR_TO_UTF8(*Line));
 	}
 
diff --git a/Engine/Source/Programs/TraceQuery/Private/JsonTraceAnalyzer.h b/Engine/Source/Programs/TraceQuery/Private/JsonTraceAnalyzer.h
--- a/Engine/Source/Programs/TraceQuery/Private/JsonTraceAnalyzer.h
+++ b/Engine/Source/Programs/TraceQuery/Private/JsonTraceAnalyzer.h
@@ -65,6 +65,12 @@ public:
 	// Resolve a TimerIndex (which may be a ~MetadataId) to its scope name, or nullptr
 	const FString* LookupTimerName(uint32 TimerId) const;
 
+	// Resolve a ThreadId to its human-readable name (e.g. "GameThread"), or nullptr
+	const FString* LookupThreadName(uint32 ThreadId) const;
+
+	// Store a thread name (called by FJsonTraceAnalyzer::OnThreadInfo)
+	void SetThreadName(uint32 ThreadId, const FString& Name);
+
 private:
 	struct FScopeEnter
 	{
@@ -108,6 +114,7 @@ private:
 	uint32 AddTimerInternal(FStringView Name);
 
 	TArray<TOptional<FString>> TimerNames;
+	TMap<uint32, FString> ThreadNames;
 	TMap<uint32, TUniquePtr<FThread>> Threads;
 	TArray<FMetadata> Metadatas;
 };
@@ -203,13 +210,14 @@ public:
 	uint64 NumCpuEvents = 0;
 	uint64 NumTotalEvents = 0;
 
-	// Set by TraceQuery.cpp to enable timer name resolution in OnAnalysisEnd
-	const FTraceQueryCpuProfilerProvider* CpuProfilerProvider = nullptr;
+	// Set by TraceQuery.cpp to enable timer/thread name resolution in OnAnalysisEnd
+	FTraceQueryCpuProfilerProvider* CpuProfilerProvider = nullptr;
 
 private:
 	// IAnalyzer interface
 	virtual void OnAnalysisBegin(const FOnAnalysisContext& Context) override;
 	virtual void OnAnalysisEnd() override;
+	virtual void OnThreadInfo(const FThreadInfo& ThreadInfo) override;
 	virtual bool OnEvent(uint16 RouteId, EStyle Style, const FOnEventContext& Context) override;
 
 	// Stats decode (kept — no factory exists for Stats)
